Added initial Simulation Template Category (still hidden) and a Simulation Blank template (BP & C++)

diff --git a/Engine/Source/Editor/MainFrame/Private/MainFrameModule.h b/Engine/Source/Editor/MainFrame/Private/MainFrameModule.h
--- a/Engine/Source/Editor/MainFrame/Private/MainFrameModule.h
+++ b/Engine/Source/Editor/MainFrame/Private/MainFrameModule.h
@@ -210,7 +210,7 @@ protected:
 public:
 
 	/** Get the size of the project browser window */
-	static FVector2D GetProjectBrowserWindowSize() { return FVector2D(1190, 733); }
+	static FVector2D GetProjectBrowserWindowSize() { return FVector2D(1190, 822); }
 
 private:
 
diff --git a/Templates/TP_SIM_Blank/Config/DefaultEngine.ini b/Templates/TP_SIM_Blank/Config/DefaultEngine.ini
new file mode 100644
--- /dev/null
+++ b/Templates/TP_SIM_Blank/Config/DefaultEngine.ini
@@ -0,0 +1,55 @@
+
+
+[/Script/EngineSettings.GameMapsSettings]
+EditorStartupMap=/Game/SimDefault/Levels/SimDefault
++EditorTemplateMapOverrides=(Thumbnail="/Game/SimDefault/Levels/Thumbnails/SimDefault_Thumbnail.SimDefault_Thumbnail",Map="/Game/SimDefault/Levels/SimDefault.SimDefault",DisplayName=NSLOCTEXT("[/Script/EngineSettings]", "C47C7345459213852543F1906820C22C", "Simulation Default"))
+GameDefaultMap=/Game/Main
+
+[/Script/Engine.RendererSettings]
+r.DefaultFeature.AutoExposure.ExtendDefaultLuminanceRange=True
+r.AllowStaticLighting=False
+r.Shadow.Virtual.Enable=1
+
+[/Script/AndroidFileServerEditor.AndroidFileServerRuntimeSettings]
+bEnablePlugin=True
+bAllowNetworkConnection=True
+SecurityToken=7A2FAE694723660CF26AA1B1F56504C8
+bIncludeInShipping=False
+bAllowExternalStartInShipping=False
+bCompileAFSProject=False
+bUseCompression=False
+bLogFiles=False
+bReportStats=False
+ConnectionType=USBOnly
+bUseManualIPAddress=False
+ManualIPAddress=
+
+[/Script/WindowsTargetPlatform.WindowsTargetSettings]
+DefaultGraphicsRHI=DefaultGraphicsRHI_DX12
+-D3D12TargetedShaderFormats=PCD3D_SM5
++D3D12TargetedShaderFormats=PCD3D_SM6
+-D3D11TargetedShaderFormats=PCD3D_SM5
++D3D11TargetedShaderFormats=PCD3D_SM5
+Compiler=Default
+AudioSampleRate=48000
+AudioCallbackBufferFrameSize=1024
+AudioNumBuffersToEnqueue=1
+AudioMaxChannels=0
+AudioNumSourceWorkers=4
+SpatializationPlugin=
+SourceDataOverridePlugin=
+ReverbPlugin=
+OcclusionPlugin=
+CompressionOverrides=(bOverrideCompressionTimes=False,DurationThreshold=5.000000,MaxNumRandomBranches=0,SoundCueQualityIndex=0)
+CacheSizeKB=65536
+MaxChunkSizeOverrideKB=0
+bResampleForDevice=False
+MaxSampleRate=48000.000000
+HighSampleRate=32000.000000
+MedSampleRate=24000.000000
+LowSampleRate=12000.000000
+MinSampleRate=8000.000000
+CompressionQualityModifier=1.000000
+AutoStreamingThreshold=0.000000
+SoundCueCookQualityIndex=-1
+
diff --git a/Templates/TP_SIM_Blank/Config/DefaultGame.ini b/Templates/TP_SIM_Blank/Config/DefaultGame.ini
new file mode 100644
--- /dev/null
+++ b/Templates/TP_SIM_Blank/Config/DefaultGame.ini
@@ -0,0 +1,3 @@
+
+[/Script/EngineSettings.GeneralProjectSettings]
+ProjectName=Blank
\ No newline at end of file
diff --git a/Templates/TP_SIM_Blank/Config/TemplateDefs.ini b/Templates/TP_SIM_Blank/Config/TemplateDefs.ini
new file mode 100644
--- /dev/null
+++ b/Templates/TP_SIM_Blank/Config/TemplateDefs.ini
@@ -0,0 +1,69 @@
+﻿; This file specifies how a template will be treated when creating a project from it.
+; It includes rules for ignoring files, replacing files and changing the contents of files.
+; This file is annotated to help with creating your own templates.
+
+[/Script/GameProjectGeneration.TemplateProjectDefs]
+
+; The categories that this template should appear in. The values in this list refer to the Key value given in TemplateCategories.ini.
+Categories=SIM
+
+; The localized display names for this template. Each entry specifies a Language code and Unicode text to display in the New Project dialog.
+LocalizedDisplayNames=(Language="en", Text="Blank")
+LocalizedDisplayNames=(Language="zh-Hans", Text="空白")
+LocalizedDisplayNames=(Language="ko", Text="기본")
+
+; The localized descriptions for this template, which is displayed when the template is selected.
+LocalizedDescriptions=(Language="en", Text="Blank template with minimal functionality and content.\n\nThis template enables you to create new outdoor scenes with appropriate lighting and georeferencing.")
+LocalizedDescriptions=(Language="zh-Hans", Text="不含任何代码的空白项目。")
+LocalizedDescriptions=(Language="ja", Text="まったく空のプロジェクトです。コードは入っていません。")
+LocalizedDescriptions=(Language="ko", Text="코드가 포함되지 않은 기본 빈 프로젝트입니다.")
+
+; The sort key to use for the template to bypass the default alphabetical sorting.
+SortKey="_1"
+
+
+; This project uses a thumbnail as the project icon
+bThumbnailAsIcon=true
+
+; A list of folders to avoid copying when creating a project from this template. 
+; These should contain files only of concern to the template itself and not desired in the user's project.
+FoldersToIgnore=Binaries
+FoldersToIgnore=Build
+FoldersToIgnore=Intermediate
+FoldersToIgnore=Saved
+FoldersToIgnore=Media
+
+; Is this a blank template? Blank templates do not specify a default map.
+bIsBlank=True
+
+; The files to ignore copying, relative to the template root folder. 
+; %TEMPLATENAME% will be replaced by the template's folder name when resolving these. 
+; Filename comparisons are case-insensitive. 
+; For example, if this template is "TP_Blank" and the user creates "MyProject", then "TP_Blank.uproject" will be copied to "MyProject.uproject".
+FilesToIgnore="%TEMPLATENAME%.uproject"
+FilesToIgnore="%TEMPLATENAME%.png"
+FilesToIgnore="Config/TemplateDefs.ini"
+FilesToIgnore="Config/config.ini"
+FilesToIgnore="%TEMPLATENAME%.opensdf"
+FilesToIgnore="%TEMPLATENAME%.sdf"
+FilesToIgnore="%TEMPLATENAME%.v11.suo"
+FilesToIgnore="%TEMPLATENAME%.v12.suo"
+FilesToIgnore="%TEMPLATENAME%.sln"
+FilesToIgnore="Manifest.json"
+FilesToIgnore="contents.txt"
+
+
+FolderRenames=(From="Source/%TEMPLATENAME%",To="Source/%PROJECTNAME%")
+FolderRenames=(From="Source/%TEMPLATENAME%Editor",To="Source/%PROJECTNAME%Editor")
+
+FilenameReplacements=(Extensions=("cpp","h","ini","cs"),From="%TEMPLATENAME_UPPERCASE%",To="%PROJECTNAME_UPPERCASE%",bCaseSensitive=true)
+FilenameReplacements=(Extensions=("cpp","h","ini","cs"),From="%TEMPLATENAME_LOWERCASE%",To="%PROJECTNAME_LOWERCASE%",bCaseSensitive=true)
+FilenameReplacements=(Extensions=("cpp","h","ini","cs"),From="%TEMPLATENAME%",To="%PROJECTNAME%",bCaseSensitive=false)
+
+ReplacementsInFiles=(Extensions=("cpp","h","ini","cs"),From="%TEMPLATENAME_UPPERCASE%",To="%PROJECTNAME_UPPERCASE%",bCaseSensitive=true)
+ReplacementsInFiles=(Extensions=("cpp","h","ini","cs"),From="%TEMPLATENAME_LOWERCASE%",To="%PROJECTNAME_LOWERCASE%",bCaseSensitive=true)
+ReplacementsInFiles=(Extensions=("cpp","h","ini","cs"),From="%TEMPLATENAME%",To="%PROJECTNAME%",bCaseSensitive=false)
+
+; Shared feature packs. The files in these packs listed in these structures marked as 'additionalfiles' will be copied on project generation.
+;SharedContentPacks=(MountName="2DSideScroller",DetailLevels=("Standard"))
+SharedContentPacks=(MountName="SimDefault",DetailLevels=("Standard"))
\ No newline at end of file
diff --git a/Templates/TP_SIM_Blank/Source/TP_SIM_Blank.Target.cs b/Templates/TP_SIM_Blank/Source/TP_SIM_Blank.Target.cs
new file mode 100644
--- /dev/null
+++ b/Templates/TP_SIM_Blank/Source/TP_SIM_Blank.Target.cs
@@ -0,0 +1,14 @@
+// Copyright Epic Games, Inc. All Rights Reserved.
+
+using UnrealBuildTool;
+using System.Collections.Generic;
+
+public class TP_SIM_BlankTarget : TargetRules
+{
+	public TP_SIM_BlankTarget( TargetInfo Target) : base(Target)
+	{
+		Type = TargetType.Game;
+		DefaultBuildSettings = BuildSettingsVersion.V2;
+		ExtraModuleNames.AddRange( new string[] { "TP_SIM_Blank" } );
+	}
+}
diff --git a/Templates/TP_SIM_Blank/Source/TP_SIM_Blank/TP_SIM_Blank.Build.cs b/Templates/TP_SIM_Blank/Source/TP_SIM_Blank/TP_SIM_Blank.Build.cs
new file mode 100644
--- /dev/null
+++ b/Templates/TP_SIM_Blank/Source/TP_SIM_Blank/TP_SIM_Blank.Build.cs
@@ -0,0 +1,23 @@
+// Copyright Epic Games, Inc. All Rights Reserved.
+
+using UnrealBuildTool;
+
+public class TP_SIM_Blank : ModuleRules
+{
+	public TP_SIM_Blank(ReadOnlyTargetRules Target) : base(Target)
+	{
+		PCHUsage = PCHUsageMode.UseExplicitOrSharedPCHs;
+	
+		PublicDependencyModuleNames.AddRange(new string[] { "Core", "CoreUObject", "Engine", "InputCore" });
+
+		PrivateDependencyModuleNames.AddRange(new string[] {  });
+
+		// Uncomment if you are using Slate UI
+		// PrivateDependencyModuleNames.AddRange(new string[] { "Slate", "SlateCore" });
+		
+		// Uncomment if you are using online features
+		// PrivateDependencyModuleNames.Add("OnlineSubsystem");
+
+		// To include OnlineSubsystemSteam, add it to the plugins section in your uproject file with the Enabled attribute set to true
+	}
+}
diff --git a/Templates/TP_SIM_Blank/Source/TP_SIM_Blank/TP_SIM_Blank.cpp b/Templates/TP_SIM_Blank/Source/TP_SIM_Blank/TP_SIM_Blank.cpp
new file mode 100644
--- /dev/null
+++ b/Templates/TP_SIM_Blank/Source/TP_SIM_Blank/TP_SIM_Blank.cpp
@@ -0,0 +1,6 @@
+// Copyright Epic Games, Inc. All Rights Reserved.
+
+#include "TP_SIM_Blank.h"
+#include "Modules/ModuleManager.h"
+
+IMPLEMENT_PRIMARY_GAME_MODULE( FDefaultGameModuleImpl, TP_SIM_Blank, "TP_SIM_Blank" );
diff --git a/Templates/TP_SIM_Blank/Source/TP_SIM_Blank/TP_SIM_Blank.h b/Templates/TP_SIM_Blank/Source/TP_SIM_Blank/TP_SIM_Blank.h
new file mode 100644
--- /dev/null
+++ b/Templates/TP_SIM_Blank/Source/TP_SIM_Blank/TP_SIM_Blank.h
@@ -0,0 +1,6 @@
+// Copyright Epic Games, Inc. All Rights Reserved.
+
+#pragma once
+
+#include "CoreMinimal.h"
+
diff --git a/Templates/TP_SIM_Blank/Source/TP_SIM_Blank/TP_SIM_BlankGameModeBase.cpp b/Templates/TP_SIM_Blank/Source/TP_SIM_Blank/TP_SIM_BlankGameModeBase.cpp
new file mode 100644
--- /dev/null
+++ b/Templates/TP_SIM_Blank/Source/TP_SIM_Blank/TP_SIM_BlankGameModeBase.cpp
@@ -0,0 +1,5 @@
+// Copyright Epic Games, Inc. All Rights Reserved.
+
+
+#include "TP_SIM_BlankGameModeBase.h"
+
diff --git a/Templates/TP_SIM_Blank/Source/TP_SIM_Blank/TP_SIM_BlankGameModeBase.h b/Templates/TP_SIM_Blank/Source/TP_SIM_Blank/TP_SIM_BlankGameModeBase.h
new file mode 100644
--- /dev/null
+++ b/Templates/TP_SIM_Blank/Source/TP_SIM_Blank/TP_SIM_BlankGameModeBase.h
@@ -0,0 +1,17 @@
+// Copyright Epic Games, Inc. All Rights Reserved.
+
+#pragma once
+
+#include "CoreMinimal.h"
+#include "GameFramework/GameModeBase.h"
+#include "TP_SIM_BlankGameModeBase.generated.h"
+
+/**
+ * 
+ */
+UCLASS()
+class TP_SIM_BLANK_API ATP_SIM_BlankGameModeBase : public AGameModeBase
+{
+	GENERATED_BODY()
+	
+};
diff --git a/Templates/TP_SIM_Blank/Source/TP_SIM_BlankEditor.Target.cs b/Templates/TP_SIM_Blank/Source/TP_SIM_BlankEditor.Target.cs
new file mode 100644
--- /dev/null
+++ b/Templates/TP_SIM_Blank/Source/TP_SIM_BlankEditor.Target.cs
@@ -0,0 +1,14 @@
+// Copyright Epic Games, Inc. All Rights Reserved.
+
+using UnrealBuildTool;
+using System.Collections.Generic;
+
+public class TP_SIM_BlankEditorTarget : TargetRules
+{
+	public TP_SIM_BlankEditorTarget( TargetInfo Target) : base(Target)
+	{
+		Type = TargetType.Editor;
+		DefaultBuildSettings = BuildSettingsVersion.V2;
+		ExtraModuleNames.AddRange( new string[] { "TP_SIM_Blank" } );
+	}
+}
diff --git a/Templates/TP_SIM_Blank/TP_SIM_Blank.uproject b/Templates/TP_SIM_Blank/TP_SIM_Blank.uproject
new file mode 100644
--- /dev/null
+++ b/Templates/TP_SIM_Blank/TP_SIM_Blank.uproject
@@ -0,0 +1,30 @@
+{
+	"FileVersion": 3,
+	"EngineAssociation": "",
+	"Category": "",
+	"Description": "",
+	"Modules": [
+		{
+			"Name": "TP_SIM_Blank",
+			"Type": "Runtime",
+			"LoadingPhase": "Default"
+		}
+	],
+	"Plugins": [
+		{
+			"Name": "ModelingToolsEditorMode",
+			"Enabled": true,
+			"TargetAllowList": [
+				"Editor"
+			]
+		},
+		{
+			"Name": "SunPosition",
+			"Enabled": true
+		},
+		{
+			"Name": "GeoReferencing",
+			"Enabled": true
+		}
+	]
+}
\ No newline at end of file
diff --git a/Templates/TP_SIM_BlankBP/Config/DefaultEngine.ini b/Templates/TP_SIM_BlankBP/Config/DefaultEngine.ini
new file mode 100644
--- /dev/null
+++ b/Templates/TP_SIM_BlankBP/Config/DefaultEngine.ini
@@ -0,0 +1,55 @@
+
+
+[/Script/EngineSettings.GameMapsSettings]
+EditorStartupMap=/Game/SimDefault/Levels/SimDefault
++EditorTemplateMapOverrides=(Thumbnail="/Game/SimDefault/Levels/Thumbnails/SimDefault_Thumbnail.SimDefault_Thumbnail",Map="/Game/SimDefault/Levels/SimDefault.SimDefault",DisplayName=NSLOCTEXT("[/Script/EngineSettings]", "C47C7345459213852543F1906820C22C", "Simulation Default"))
+GameDefaultMap=/Game/Main
+
+[/Script/Engine.RendererSettings]
+r.DefaultFeature.AutoExposure.ExtendDefaultLuminanceRange=True
+r.AllowStaticLighting=False
+r.Shadow.Virtual.Enable=1
+
+[/Script/AndroidFileServerEditor.AndroidFileServerRuntimeSettings]
+bEnablePlugin=True
+bAllowNetworkConnection=True
+SecurityToken=7A2FAE694723660CF26AA1B1F56504C8
+bIncludeInShipping=False
+bAllowExternalStartInShipping=False
+bCompileAFSProject=False
+bUseCompression=False
+bLogFiles=False
+bReportStats=False
+ConnectionType=USBOnly
+bUseManualIPAddress=False
+ManualIPAddress=
+
+[/Script/WindowsTargetPlatform.WindowsTargetSettings]
+DefaultGraphicsRHI=DefaultGraphicsRHI_DX12
+-D3D12TargetedShaderFormats=PCD3D_SM5
++D3D12TargetedShaderFormats=PCD3D_SM6
+-D3D11TargetedShaderFormats=PCD3D_SM5
++D3D11TargetedShaderFormats=PCD3D_SM5
+Compiler=Default
+AudioSampleRate=48000
+AudioCallbackBufferFrameSize=1024
+AudioNumBuffersToEnqueue=1
+AudioMaxChannels=0
+AudioNumSourceWorkers=4
+SpatializationPlugin=
+SourceDataOverridePlugin=
+ReverbPlugin=
+OcclusionPlugin=
+CompressionOverrides=(bOverrideCompressionTimes=False,DurationThreshold=5.000000,MaxNumRandomBranches=0,SoundCueQualityIndex=0)
+CacheSizeKB=65536
+MaxChunkSizeOverrideKB=0
+bResampleForDevice=False
+MaxSampleRate=48000.000000
+HighSampleRate=32000.000000
+MedSampleRate=24000.000000
+LowSampleRate=12000.000000
+MinSampleRate=8000.000000
+CompressionQualityModifier=1.000000
+AutoStreamingThreshold=0.000000
+SoundCueCookQualityIndex=-1
+
diff --git a/Templates/TP_SIM_BlankBP/Config/DefaultGame.ini b/Templates/TP_SIM_BlankBP/Config/DefaultGame.ini
new file mode 100644
--- /dev/null
+++ b/Templates/TP_SIM_BlankBP/Config/DefaultGame.ini
@@ -0,0 +1,3 @@
+
+[/Script/EngineSettings.GeneralProjectSettings]
+ProjectName=Blank
\ No newline at end of file
diff --git a/Templates/TP_SIM_BlankBP/Config/TemplateDefs.ini b/Templates/TP_SIM_BlankBP/Config/TemplateDefs.ini
new file mode 100644
--- /dev/null
+++ b/Templates/TP_SIM_BlankBP/Config/TemplateDefs.ini
@@ -0,0 +1,69 @@
+﻿; This file specifies how a template will be treated when creating a project from it.
+; It includes rules for ignoring files, replacing files and changing the contents of files.
+; This file is annotated to help with creating your own templates.
+
+[/Script/GameProjectGeneration.TemplateProjectDefs]
+
+; The categories that this template should appear in. The values in this list refer to the Key value given in TemplateCategories.ini.
+Categories=SIM
+
+; The localized display names for this template. Each entry specifies a Language code and Unicode text to display in the New Project dialog.
+LocalizedDisplayNames=(Language="en", Text="Blank")
+LocalizedDisplayNames=(Language="zh-Hans", Text="空白")
+LocalizedDisplayNames=(Language="ko", Text="기본")
+
+; The localized descriptions for this template, which is displayed when the template is selected.
+LocalizedDescriptions=(Language="en", Text="Blank template with minimal functionality and content.\n\nThis template enables you to create new outdoor scenes with appropriate lighting and georeferencing.")
+LocalizedDescriptions=(Language="zh-Hans", Text="不含任何代码的空白项目。")
+LocalizedDescriptions=(Language="ja", Text="まったく空のプロジェクトです。コードは入っていません。")
+LocalizedDescriptions=(Language="ko", Text="코드가 포함되지 않은 기본 빈 프로젝트입니다.")
+
+; The sort key to use for the template to bypass the default alphabetical sorting.
+SortKey="_1"
+
+
+; This project uses a thumbnail as the project icon
+bThumbnailAsIcon=true
+
+; A list of folders to avoid copying when creating a project from this template. 
+; These should contain files only of concern to the template itself and not desired in the user's project.
+FoldersToIgnore=Binaries
+FoldersToIgnore=Build
+FoldersToIgnore=Intermediate
+FoldersToIgnore=Saved
+FoldersToIgnore=Media
+
+; Is this a blank template? Blank templates do not specify a default map.
+bIsBlank=True
+
+; The files to ignore copying, relative to the template root folder. 
+; %TEMPLATENAME% will be replaced by the template's folder name when resolving these. 
+; Filename comparisons are case-insensitive. 
+; For example, if this template is "TP_Blank" and the user creates "MyProject", then "TP_Blank.uproject" will be copied to "MyProject.uproject".
+FilesToIgnore="%TEMPLATENAME%.uproject"
+FilesToIgnore="%TEMPLATENAME%.png"
+FilesToIgnore="Config/TemplateDefs.ini"
+FilesToIgnore="Config/config.ini"
+FilesToIgnore="%TEMPLATENAME%.opensdf"
+FilesToIgnore="%TEMPLATENAME%.sdf"
+FilesToIgnore="%TEMPLATENAME%.v11.suo"
+FilesToIgnore="%TEMPLATENAME%.v12.suo"
+FilesToIgnore="%TEMPLATENAME%.sln"
+FilesToIgnore="Manifest.json"
+FilesToIgnore="contents.txt"
+
+
+FolderRenames=(From="Source/%TEMPLATENAME%",To="Source/%PROJECTNAME%")
+FolderRenames=(From="Source/%TEMPLATENAME%Editor",To="Source/%PROJECTNAME%Editor")
+
+FilenameReplacements=(Extensions=("cpp","h","ini","cs"),From="%TEMPLATENAME_UPPERCASE%",To="%PROJECTNAME_UPPERCASE%",bCaseSensitive=true)
+FilenameReplacements=(Extensions=("cpp","h","ini","cs"),From="%TEMPLATENAME_LOWERCASE%",To="%PROJECTNAME_LOWERCASE%",bCaseSensitive=true)
+FilenameReplacements=(Extensions=("cpp","h","ini","cs"),From="%TEMPLATENAME%",To="%PROJECTNAME%",bCaseSensitive=false)
+
+ReplacementsInFiles=(Extensions=("cpp","h","ini","cs"),From="%TEMPLATENAME_UPPERCASE%",To="%PROJECTNAME_UPPERCASE%",bCaseSensitive=true)
+ReplacementsInFiles=(Extensions=("cpp","h","ini","cs"),From="%TEMPLATENAME_LOWERCASE%",To="%PROJECTNAME_LOWERCASE%",bCaseSensitive=true)
+ReplacementsInFiles=(Extensions=("cpp","h","ini","cs"),From="%TEMPLATENAME%",To="%PROJECTNAME%",bCaseSensitive=false)
+
+; Shared feature packs. The files in these packs listed in these structures marked as 'additionalfiles' will be copied on project generation.
+;SharedContentPacks=(MountName="2DSideScroller",DetailLevels=("Standard"))
+SharedContentPacks=(MountName="SimDefault",DetailLevels=("Standard"))
\ No newline at end of file
diff --git a/Templates/TP_SIM_BlankBP/TP_SIM_BlankBP.uproject b/Templates/TP_SIM_BlankBP/TP_SIM_BlankBP.uproject
new file mode 100644
--- /dev/null
+++ b/Templates/TP_SIM_BlankBP/TP_SIM_BlankBP.uproject
@@ -0,0 +1,23 @@
+{
+	"FileVersion": 3,
+	"EngineAssociation": "",
+	"Category": "",
+	"Description": "",
+	"Plugins": [
+		{
+			"Name": "ModelingToolsEditorMode",
+			"Enabled": true,
+			"TargetAllowList": [
+				"Editor"
+			]
+		},
+		{
+			"Name": "SunPosition",
+			"Enabled": true
+		},
+		{
+			"Name": "GeoReferencing",
+			"Enabled": true
+		}
+	]
+}
\ No newline at end of file
diff --git a/Templates/TemplateCategories.ini b/Templates/TemplateCategories.ini
--- a/Templates/TemplateCategories.ini
+++ b/Templates/TemplateCategories.ini
@@ -19,3 +19,5 @@ Categories=(Key="ME", LocalizedDisplayNames=((Language="en",Text="Film / Video &
 Categories=(Key="AEC", LocalizedDisplayNames=((Language="en",Text="Architecture"),(Language="zh-Hans",Text="建筑"),(Language="ja",Text="建築"),(Language="ko",Text="건축")), LocalizedDescriptions=((Language="en",Text="Select a starting point for multi-user design reviews, photorealistic architectural design visualizations, sunlight studies, or stylized renderings."),(Language="zh-Hans", Text="为多用户设计评审、照片级建筑设计可视化、阳光研究或风格化渲染选择一个起点。"),(Language="ja",Text="複数のユーザーによるレビューや、フォトリアルな建築デザインのビジュアライゼーション、太陽光の調査、スタイライズド レンダリングのための土台を選択することができます。"),(Language="ko",Text="멀티 유저 디자인 리뷰, 건축 디자인의 포토리얼한 시각화, 선라이트 스터디 또는 스타일라이즈드 렌더링을 시작해 보세요.")), Icon="Media/AEC_2x.png", IsMajorCategory=true)
 
 Categories=(Key="MFG", LocalizedDisplayNames=((Language="en",Text="Automotive\nProduct Design &\nManufacturing"),(Language="zh-Hans",Text="汽车、产品设计和制造"),(Language="ja",Text="自動車、プロダクト デザイン、製造"),(Language="ko",Text="자동차, 제품 디자인 및 제조")), LocalizedDescriptions=((Language="en",Text="Find templates for multi-user design reviews, photobooth studio environments, and product configurators."),(Language="zh-Hans",Text="为多用户设计评审、Photobooth工作室环境和产品配置器寻找适合的模板。"),(Language="ja",Text="複数のユーザーによるレビューや、フォトブース スタジオの環境背景、製品コンフィギュレーターのためのテンプレートを見つけることができます。"),(Language="ko",Text="멀티 유저 디자인 리뷰, 포토부스 스튜디오 환경 및 제품 컨피규레이터를 위한 템플릿을 살펴보세요.")), Icon="Media/AutomotiveProduct_2x.png", IsMajorCategory=true)
+
+;Categories=(Key="SIM", LocalizedDisplayNames=((Language="en",Text="Simulation"),(Language="zh-Hans",Text="TODO_Simulation_zh"),(Language="ja",Text="TODO_Simulation_ja"),(Language="ko",Text="TODO_Simulation_ko")), LocalizedDescriptions=((Language="en",Text="Find templates for Outdoor georeferenced environments, First person, Driving, Multichannels and AR/VR projects."),(Language="zh-Hans",Text="TODO_SimulationDesc_"),(Language="ja",Text="TODO_SimulationDesc_ja"),(Language="ko",Text="TODO_SimulationDesc_ko")), Icon="Media/Simulation_2x.png", IsMajorCategory=true)
diff --git a/Templates/TemplateResources/Standard/SimDefault/FeaturePack/manifest.json b/Templates/TemplateResources/Standard/SimDefault/FeaturePack/manifest.json
new file mode 100644
--- /dev/null
+++ b/Templates/TemplateResources/Standard/SimDefault/FeaturePack/manifest.json
@@ -0,0 +1,48 @@
+{
+	"Version": 1,
+	"Name":
+	[
+		{
+			"Language": "en",
+			"Text": "Simulation templates Shared Resource Standard"
+		}
+	],
+	"Description":
+	[
+		{
+			"Language": "en",
+			"Text": "Shared Resource Pack for Simulation templates"
+		}
+	],
+	"AssetTypes":
+	[
+		{
+			"Language": "en",
+			"Text": ""
+		}
+	],
+	"SearchTags":
+	[
+		{
+			"Language": "en",
+			"Text": ""
+		}
+	],	
+	"ClassTypes": " ",
+	"Category": "SharedPack",
+	"Thumbnail": "SimDefault.png",
+	"Screenshots":
+	[
+		"SimDefault_Preview.png"
+	],
+	"Ident": "SimDefaultStandard",
+	"AdditionalFiles":
+	{
+		"DestinationFilesFolder": "SimDefault",
+		"AdditionalFilesList":
+		[
+			"Templates/TemplateResources/Standard/SimDefault/Content/*.*"
+		]
+	}
+}
+
