Update template title description translation

diff --git a/Templates/TP_SIM_Blank/Config/TemplateDefs.ini b/Templates/TP_SIM_Blank/Config/TemplateDefs.ini
--- a/Templates/TP_SIM_Blank/Config/TemplateDefs.ini
+++ b/Templates/TP_SIM_Blank/Config/TemplateDefs.ini
@@ -9,14 +9,15 @@ Categories=SIM
 
 ; The localized display names for this template. Each entry specifies a Language code and Unicode text to display in the New Project dialog.
 LocalizedDisplayNames=(Language="en", Text="Simulation Blank")
-LocalizedDisplayNames=(Language="zh-Hans", Text="空白")
-LocalizedDisplayNames=(Language="ko", Text="기본")
+LocalizedDisplayNames=(Language="zh-Hans", Text="模拟-空白")
+LocalizedDisplayNames=(Language="ko", Text="시뮬레이션 기본")
+LocalizedDisplayNames=(Language="ja", Text="シミュレーション Blank")
 
 ; The localized descriptions for this template, which is displayed when the template is selected.
 LocalizedDescriptions=(Language="en", Text="A simple empty project template enabling initial settings and plugins. A starting point suitable for most of your simulation applications.\n\nEarth Atmosphere, Lighting, Volumetric Clouds, Geo Coordinates, WGS84")
-LocalizedDescriptions=(Language="zh-Hans", Text="不含任何代码的空白项目。")
-LocalizedDescriptions=(Language="ja", Text="まったく空のプロジェクトです。コードは入っていません。")
-LocalizedDescriptions=(Language="ko", Text="코드가 포함되지 않은 기본 빈 프로젝트입니다.")
+LocalizedDescriptions=(Language="zh-Hans", Text="简易的空白项目模板，提供了一些初始设置和插件。可以作为你在创建多数模拟仿真应用时的起点。\n\n地球大气、光照、体积云、地理坐标、WGS84")
+LocalizedDescriptions=(Language="ja", Text="初期設定とプラグインを有効にするシンプルな空のプロジェクト。ほとんどのシミュレーション アプリケーションを始めるのに適しています。\n\n地球大気、ライティング、ボリューメトリック クラウド、地理座標、WGS84")
+LocalizedDescriptions=(Language="ko", Text="초기 세팅과 플러그인을 활성화하는 간단한 기본 프로젝트 템플릿으로 대부분의 시뮬레이션 애플리케이션에 시작점으로 사용하기 적합합니다.\n\n지구 애트머스피어, 라이팅, 볼류메트릭 클라우드, 지오 좌표, WGS84")
 
 ; The sort key to use for the template to bypass the default alphabetical sorting.
 SortKey="_1"
diff --git a/Templates/TP_SIM_Blank/manifest.json b/Templates/TP_SIM_Blank/manifest.json
--- a/Templates/TP_SIM_Blank/manifest.json
+++ b/Templates/TP_SIM_Blank/manifest.json
@@ -8,15 +8,15 @@
 		},
 		{
 			"Language": "ko",
-			"Text": "일인칭"
+			"Text": "시뮬레이션 기본"
 		},
 		{
 			"Language": "zh-Hans",
-			"Text": "第一人称游戏"
+			"Text": "模拟-空白"
 		},
 		{
 			"Language": "ja",
-			"Text": "ファーストパーソン"
+			"Text": "シミュレーション Blank"
 		}
 	],
 	"Description":
@@ -27,15 +27,15 @@
 		},
 		{
 			"Language": "ko",
-			"Text": "일인칭 템플릿은 일인칭 시점에서 보이는 (양팔로 표현되는) 플레이어 캐릭터를 특징으로 합니다. 캐릭터 이동은 키보드, 컨트롤러, 터치 디바이스의 가상 조이스틱으로 가능합니다. 추가적으로 플레이어는 마우스, 컨트롤러, 터치 디바이스의 가상 조이스틱을 사용하여 주변을 둘러볼 수도 있습니다. 캐릭터에는 총이 장착되어 있어 마우스, 컨트롤러, 터치 디바이스의 가상 조이스틱을 사용하면 간단한 발사체가 발사되어, 아레나 벽에 튕기면서 레벨에 있는 몇몇 물리 오브젝트에 영향을 줍니다."
+			"Text": "시뮬레이션에 적합한 디폴트 세팅으로 구성된 기본 환경입니다."
 		},
 		{
 			"Language": "zh-Hans",
-			"Text": "第一人称模板包含一个第一人称玩家角色（通过一双手臂呈现）。你可以通过键盘、手柄或触摸式设备上的虚拟摇杆控制角色在关卡中四处移动。你还可以使用鼠标、手柄或触摸式设备上的虚拟摇杆转动角色视角，观察四周。此外，你还可以让角色拾取枪支，然后用鼠标、手柄或触摸式设备上的虚拟摇杆来发射球形子弹。子弹会撞飞关卡中的物理对象，还能从竞技场墙壁上弹开。"
+			"Text": "带有默认设置的简易空白场景。适合用于创建模拟仿真项目。"
 		},
 		{
 			"Language": "ja",
-			"Text": "ファーストパーソンテンプレートには、一人称視点で表示されるプレイヤーキャラクタ（腕として表示されています）がいます。キーボード、コントローラー、タッチデバイスでのバーチャルジョイスティックでこのキャラクタを移動させることができます。追加して、プレイヤーはマウス、コントローラー、タッチデバイスでのバーチャルジョイスティックで周りを見回すことができます。さらにキャラクタは銃を装備して、マウス、コントローラー、タッチデバイスでのバーチャルジョイスティックを使ってシンプルな球体の弾を発射できます。この弾は物理オブジェクトに影響し、アリーナの壁で反射します。"
+			"Text": "シミュレーション プロジェクトに適したデフォルト設定の基本的な空の環境"
 		}
 	],
 	"AssetTypes":
diff --git a/Templates/TP_SIM_BlankBP/Config/TemplateDefs.ini b/Templates/TP_SIM_BlankBP/Config/TemplateDefs.ini
--- a/Templates/TP_SIM_BlankBP/Config/TemplateDefs.ini
+++ b/Templates/TP_SIM_BlankBP/Config/TemplateDefs.ini
@@ -9,14 +9,15 @@ Categories=SIM
 
 ; The localized display names for this template. Each entry specifies a Language code and Unicode text to display in the New Project dialog.
 LocalizedDisplayNames=(Language="en", Text="Simulation Blank")
-LocalizedDisplayNames=(Language="zh-Hans", Text="空白")
-LocalizedDisplayNames=(Language="ko", Text="기본")
+LocalizedDisplayNames=(Language="zh-Hans", Text="模拟-空白")
+LocalizedDisplayNames=(Language="ko", Text="시뮬레이션 기본")
+LocalizedDisplayNames=(Language="ja", Text="シミュレーション Blank")
 
 ; The localized descriptions for this template, which is displayed when the template is selected.
 LocalizedDescriptions=(Language="en", Text="A simple empty project template enabling initial settings and plugins. A starting point suitable for most of your simulation applications.\n\nEarth Atmosphere, Lighting, Volumetric Clouds, Geo Coordinates, WGS84")
-LocalizedDescriptions=(Language="zh-Hans", Text="不含任何代码的空白项目。")
-LocalizedDescriptions=(Language="ja", Text="まったく空のプロジェクトです。コードは入っていません。")
-LocalizedDescriptions=(Language="ko", Text="코드가 포함되지 않은 기본 빈 프로젝트입니다.")
+LocalizedDescriptions=(Language="zh-Hans", Text="简易的空白项目模板，提供了一些初始设置和插件。可以作为你在创建多数模拟仿真应用时的起点。\n\n地球大气、光照、体积云、地理坐标、WGS84")
+LocalizedDescriptions=(Language="ja", Text="初期設定とプラグインを有効にするシンプルな空のプロジェクト。ほとんどのシミュレーション アプリケーションを始めるのに適しています。\n\n地球大気、ライティング、ボリューメトリック クラウド、地理座標、WGS84")
+LocalizedDescriptions=(Language="ko", Text="초기 세팅과 플러그인을 활성화하는 간단한 기본 프로젝트 템플릿으로 대부분의 시뮬레이션 애플리케이션에 시작점으로 사용하기 적합합니다.\n\n지구 애트머스피어, 라이팅, 볼류메트릭 클라우드, 지오 좌표, WGS84")
 
 ; The sort key to use for the template to bypass the default alphabetical sorting.
 SortKey="_1"
diff --git a/Templates/TP_SIM_BlankBP/manifest.json b/Templates/TP_SIM_BlankBP/manifest.json
--- a/Templates/TP_SIM_BlankBP/manifest.json
+++ b/Templates/TP_SIM_BlankBP/manifest.json
@@ -8,15 +8,15 @@
 		},
 		{
 			"Language": "ko",
-			"Text": "일인칭"
+			"Text": "시뮬레이션 기본"
 		},
 		{
 			"Language": "zh-Hans",
-			"Text": "第一人称游戏"
+			"Text": "模拟-空白"
 		},
 		{
 			"Language": "ja",
-			"Text": "ファーストパーソン"
+			"Text": "シミュレーション Blank"
 		}
 	],
 	"Description":
@@ -27,15 +27,15 @@
 		},
 		{
 			"Language": "ko",
-			"Text": "일인칭 템플릿은 일인칭 시점에서 보이는 (양팔로 표현되는) 플레이어 캐릭터를 특징으로 합니다. 캐릭터 이동은 키보드, 컨트롤러, 터치 디바이스의 가상 조이스틱으로 가능합니다. 추가적으로 플레이어는 마우스, 컨트롤러, 터치 디바이스의 가상 조이스틱을 사용하여 주변을 둘러볼 수도 있습니다. 캐릭터에는 총이 장착되어 있어 마우스, 컨트롤러, 터치 디바이스의 가상 조이스틱을 사용하면 간단한 발사체가 발사되어, 아레나 벽에 튕기면서 레벨에 있는 몇몇 물리 오브젝트에 영향을 줍니다."
+			"Text": "시뮬레이션에 적합한 디폴트 세팅으로 구성된 기본 환경입니다."
 		},
 		{
 			"Language": "zh-Hans",
-			"Text": "第一人称模板包含一个第一人称玩家角色（通过一双手臂呈现）。你可以通过键盘、手柄或触摸式设备上的虚拟摇杆控制角色在关卡中四处移动。你还可以使用鼠标、手柄或触摸式设备上的虚拟摇杆转动角色视角，观察四周。此外，你还可以让角色拾取枪支，然后用鼠标、手柄或触摸式设备上的虚拟摇杆来发射球形子弹。子弹会撞飞关卡中的物理对象，还能从竞技场墙壁上弹开。"
+			"Text": "带有默认设置的简易空白场景。适合用于创建模拟仿真项目。"
 		},
 		{
 			"Language": "ja",
-			"Text": "ファーストパーソンテンプレートには、一人称視点で表示されるプレイヤーキャラクタ（腕として表示されています）がいます。キーボード、コントローラー、タッチデバイスでのバーチャルジョイスティックでこのキャラクタを移動させることができます。追加して、プレイヤーはマウス、コントローラー、タッチデバイスでのバーチャルジョイスティックで周りを見回すことができます。さらにキャラクタは銃を装備して、マウス、コントローラー、タッチデバイスでのバーチャルジョイスティックを使ってシンプルな球体の弾を発射できます。この弾は物理オブジェクトに影響し、アリーナの壁で反射します。"
+			"Text": "シミュレーション プロジェクトに適したデフォルト設定の基本的な空の環境"
 		}
 	],
 	"AssetTypes":
